Copy of 11160369 from //UE4/Dev-Enterprise:

RemoteImport:
Add module RemoteImportLibrary: base types used in RemoteImport features

diff --git a/Engine/Plugins/Importers/RemoteImport/RemoteImport.uplugin b/Engine/Plugins/Importers/RemoteImport/RemoteImport.uplugin
new file mode 100644
--- /dev/null
+++ b/Engine/Plugins/Importers/RemoteImport/RemoteImport.uplugin
@@ -0,0 +1,34 @@
+{
+	"FileVersion": 3,
+	"Version": 1,
+	"VersionName": "1.0",
+	"FriendlyName": "Remote Import",
+	"Description": "Allows external applications to remotely trigger a content import",
+	"Category": "Importers",
+	"CreatedBy": "Epic Games, Inc.",
+	"CreatedByURL": "http://epicgames.com",
+	"DocsURL": "",
+	"MarketplaceURL": "",
+	"SupportURL": "",
+	"EnabledByDefault": false,
+	"CanContainContent": false,
+	"IsBetaVersion": true,
+	"Installed": false,
+	"Modules": [
+		{
+			"Name": "RemoteImportLibrary",
+			"Type": "Runtime",
+			"LoadingPhase": "Default"
+		}
+	],
+	"Plugins": [
+		{
+			"Name": "OnlineSubsystem",
+			"Enabled": true
+		},
+		{
+			"Name": "DatasmithImporter",
+			"Enabled": true
+		}
+	]
+}
\ No newline at end of file
diff --git a/Engine/Plugins/Importers/RemoteImport/Source/RemoteImportLibrary/Private/RemoteImportLibrary.cpp b/Engine/Plugins/Importers/RemoteImport/Source/RemoteImportLibrary/Private/RemoteImportLibrary.cpp
new file mode 100644
--- /dev/null
+++ b/Engine/Plugins/Importers/RemoteImport/Source/RemoteImportLibrary/Private/RemoteImportLibrary.cpp
@@ -0,0 +1,86 @@
+// Copyright Epic Games, Inc. All Rights Reserved.
+
+#include "RemoteImportLibrary.h"
+#include "RemoteImportLibraryLog.h"
+#include "RemoteImportServer.h"
+
+
+TArray<FRemoteImportAnchor> URemoteImportLibrary::Anchors;
+TUniquePtr<FRemoteImportServer> URemoteImportLibrary::Instance;
+FSimpleMulticastDelegate URemoteImportLibrary::OnAnchorListChange;
+
+
+URemoteImportLibrary::URemoteImportLibrary()
+{
+	if (HasAnyFlags(RF_ClassDefaultObject))
+	{
+		FRemoteImportAnchor DefaultAnchor;
+		DefaultAnchor.Name = TEXT("Default");
+		RegisterAnchor(DefaultAnchor);
+	}
+}
+
+bool URemoteImportLibrary::ImportSource(const FString& FilePath, const FString& DestinationName)
+{
+	FRemoteImportAnchor* Destination = FindAnchor(DestinationName);
+	if (!Destination)
+	{
+		UE_LOG(LogRemoteImportLibrary, Warning, TEXT("ImportSource failed: destination [%s] is unknown."), *DestinationName);
+		return false;
+	}
+
+	Destination->OnImportFileDelegate.ExecuteIfBound(FilePath);
+	return Destination->OnImportFileDelegate.IsBound();
+}
+
+
+void URemoteImportLibrary::RegisterAnchor(const FRemoteImportAnchor& Anchor)
+{
+	Anchors.Add(Anchor);
+	OnAnchorListChange.Broadcast();
+}
+
+
+void URemoteImportLibrary::UnregisterAnchor(const FString& AnchorName)
+{
+	Anchors.RemoveAll([&AnchorName](const FRemoteImportAnchor& Anchor){ return Anchor.Name == AnchorName; });
+	OnAnchorListChange.Broadcast();
+}
+
+
+TArray<FString> URemoteImportLibrary::ListAnchors()
+{
+	TArray<FString> Out;
+	Out.Reserve(Anchors.Num());
+	for (const FRemoteImportAnchor& Anchor : Anchors)
+	{
+		Out.Add(Anchor.Name);
+	}
+	return Out;
+}
+
+
+FRemoteImportAnchor* URemoteImportLibrary::FindAnchor(const FString& AnchorName)
+{
+	return Anchors.FindByPredicate([&AnchorName](const FRemoteImportAnchor& Anchor){ return Anchor.Name == AnchorName; });
+}
+
+
+void URemoteImportLibrary::StartRemoteImportServer()
+{
+	Instance = MakeUnique<FRemoteImportServer>();
+	UE_LOG(LogRemoteImportLibrary, Display, TEXT("RemoteImportServer started"));
+}
+
+
+void URemoteImportLibrary::StopRemoteImportServer()
+{
+	Instance.Reset();
+	UE_LOG(LogRemoteImportLibrary, Display, TEXT("RemoteImportServer stopped"));
+}
+
+
+bool URemoteImportLibrary::IsRemoteImportServerActive()
+{
+	return Instance.IsValid();
+}
diff --git a/Engine/Plugins/Importers/RemoteImport/Source/RemoteImportLibrary/Private/RemoteImportLibraryLog.h b/Engine/Plugins/Importers/RemoteImport/Source/RemoteImportLibrary/Private/RemoteImportLibraryLog.h
new file mode 100644
--- /dev/null
+++ b/Engine/Plugins/Importers/RemoteImport/Source/RemoteImportLibrary/Private/RemoteImportLibraryLog.h
@@ -0,0 +1,7 @@
+// Copyright Epic Games, Inc. All Rights Reserved.
+
+#pragma once
+
+#include "Logging/LogMacros.h"
+
+DECLARE_LOG_CATEGORY_EXTERN(LogRemoteImportLibrary, Log, All);
diff --git a/Engine/Plugins/Importers/RemoteImport/Source/RemoteImportLibrary/Private/RemoteImportLibraryModule.cpp b/Engine/Plugins/Importers/RemoteImport/Source/RemoteImportLibrary/Private/RemoteImportLibraryModule.cpp
new file mode 100644
--- /dev/null
+++ b/Engine/Plugins/Importers/RemoteImport/Source/RemoteImportLibrary/Private/RemoteImportLibraryModule.cpp
@@ -0,0 +1,10 @@
+// Copyright Epic Games, Inc. All Rights Reserved.
+#include "RemoteImportLibraryModule.h"
+#include "RemoteImportLibraryLog.h"
+
+#include "Modules/ModuleInterface.h"
+#include "Modules/ModuleManager.h"
+
+IMPLEMENT_MODULE(FRemoteImportLibraryModule, RemoteImportLibrary);
+
+DEFINE_LOG_CATEGORY(LogRemoteImportLibrary)
diff --git a/Engine/Plugins/Importers/RemoteImport/Source/RemoteImportLibrary/Private/RemoteImportServer.cpp b/Engine/Plugins/Importers/RemoteImport/Source/RemoteImportLibrary/Private/RemoteImportServer.cpp
new file mode 100644
--- /dev/null
+++ b/Engine/Plugins/Importers/RemoteImport/Source/RemoteImportLibrary/Private/RemoteImportServer.cpp
@@ -0,0 +1,80 @@
+// Copyright Epic Games, Inc. All Rights Reserved.
+
+#include "RemoteImportServer.h"
+#include "MessageEndpoint.h"
+#include "MessageEndpointBuilder.h"
+#include "RemoteImportMessages.h"
+#include "RemoteImportLibrary.h"
+
+
+
+DEFINE_LOG_CATEGORY_STATIC(LogRemoteImportServer, Log, All);
+
+
+
+FRemoteImportServer::FRemoteImportServer()
+{
+	auto MessageEndpointBuilder = FMessageEndpoint::Builder(TEXT("RemoteImportServer"))
+		.ReceivingOnThread(ENamedThreads::GameThread)
+		.Handling<FPingMessage>(this, &FRemoteImportServer::OnPingMessage)
+		.Handling<FImportFileRequest>(this, &FRemoteImportServer::OnImportCommandMessage)
+	;
+	MessageEndpoint = MessageEndpointBuilder.Build();
+
+	if (MessageEndpoint.IsValid())
+	{
+		MessageEndpoint->Subscribe<FPingMessage>();
+		MessageEndpoint->Subscribe<FImportFileRequest>();
+	}
+
+	OnAnchorChangeDelegateHandle = URemoteImportLibrary::GetAnchorListChangeDelegate().AddRaw(
+		this, &FRemoteImportServer::PublishAnchorList);
+}
+
+FRemoteImportServer::~FRemoteImportServer()
+{
+	URemoteImportLibrary::GetAnchorListChangeDelegate().Remove(OnAnchorChangeDelegateHandle);
+	if (MessageEndpoint.IsValid())
+	{
+		// Disable the Endpoint message handling since the message could keep it alive a bit.
+		MessageEndpoint->Disable();
+		MessageEndpoint.Reset();
+	}
+}
+
+FServerStateMessage* MakeAnchorListMessage()
+{
+	auto ServerStateMessage = new FServerStateMessage();
+	ServerStateMessage->Anchors = URemoteImportLibrary::ListAnchors();
+	return ServerStateMessage;
+}
+
+void FRemoteImportServer::OnPingMessage(const FPingMessage& Message, const TSharedRef<IMessageContext, ESPMode::ThreadSafe>& Context)
+{
+	UE_LOG(LogRemoteImportServer, Display, TEXT("OnPingMessage: %d from %s"), Message.Version, *Context->GetSender().ToString());
+	if (MessageEndpoint.IsValid())
+	{
+		MessageEndpoint->Send(new FPongMessage(FString::Printf(TEXT("Received %d"), Message.Version)), Context->GetSender());
+		MessageEndpoint->Send(MakeAnchorListMessage(), Context->GetSender());
+	}
+}
+
+void FRemoteImportServer::OnImportCommandMessage(const FImportFileRequest& Message, const TSharedRef<IMessageContext, ESPMode::ThreadSafe>& Context)
+{
+	URemoteImportLibrary::ImportSource(Message.File, Message.Destination);
+
+	if (MessageEndpoint.IsValid())
+	{
+		MessageEndpoint->Send(new FImportFileResponse(&Message), Context->GetSender());
+	}
+}
+
+void FRemoteImportServer::PublishAnchorList()
+{
+	// #ueent_todo: use a timer to avoid spam and batch multiple updates in one
+	if (MessageEndpoint.IsValid())
+	{
+		MessageEndpoint->Publish(MakeAnchorListMessage());
+	}
+}
+
diff --git a/Engine/Plugins/Importers/RemoteImport/Source/RemoteImportLibrary/Private/RemoteImportServer.h b/Engine/Plugins/Importers/RemoteImport/Source/RemoteImportLibrary/Private/RemoteImportServer.h
new file mode 100644
--- /dev/null
+++ b/Engine/Plugins/Importers/RemoteImport/Source/RemoteImportLibrary/Private/RemoteImportServer.h
@@ -0,0 +1,32 @@
+// Copyright Epic Games, Inc. All Rights Reserved.
+
+#pragma once
+
+#include "MessageEndpoint.h"
+
+struct FPingMessage;
+struct FImportFileRequest;
+class IMessageContext;
+
+
+/**
+ * Scopes the MessageBus endpoint, i.e. the ability to receive commands from RemoteImport clients.
+ */
+class FRemoteImportServer
+{
+public:
+	FRemoteImportServer();
+
+	~FRemoteImportServer();
+
+private:
+	void OnPingMessage(const FPingMessage& Message, const TSharedRef<IMessageContext, ESPMode::ThreadSafe>& Context);
+
+	void OnImportCommandMessage(const FImportFileRequest& Message, const TSharedRef<IMessageContext, ESPMode::ThreadSafe>& Context);
+
+	void PublishAnchorList();
+
+	TSharedPtr<FMessageEndpoint, ESPMode::ThreadSafe> MessageEndpoint;
+
+	FDelegateHandle OnAnchorChangeDelegateHandle;
+};
diff --git a/Engine/Plugins/Importers/RemoteImport/Source/RemoteImportLibrary/Public/RemoteImportLibrary.h b/Engine/Plugins/Importers/RemoteImport/Source/RemoteImportLibrary/Public/RemoteImportLibrary.h
new file mode 100644
--- /dev/null
+++ b/Engine/Plugins/Importers/RemoteImport/Source/RemoteImportLibrary/Public/RemoteImportLibrary.h
@@ -0,0 +1,112 @@
+// Copyright Epic Games, Inc. All Rights Reserved.
+
+#pragma once
+
+#include "CoreMinimal.h"
+#include "Delegates/Delegate.h"
+#include "Kismet/BlueprintFunctionLibrary.h"
+
+#include "RemoteImportLibrary.generated.h"
+
+class FRemoteImportServer;
+
+
+DECLARE_DELEGATE_OneParam(FStringArgDelegate, const FString&);
+
+
+/**
+ * Structure describing a RemoteImport destination.
+ *
+ * Server-side, Anchors are setup to implement RemoteImport commands.
+ * By registering Anchors to the URemoteImportLibrary, Anchors becomes visible from clients.
+ * Clients can send commands to registered Anchors.
+ */
+USTRUCT(BlueprintType, meta=(Experimental))
+struct REMOTEIMPORTLIBRARY_API FRemoteImportAnchor
+{
+	GENERATED_BODY()
+
+public:
+	/** User friendly name, exposed in client UIs */
+	UPROPERTY(BlueprintReadWrite, Category=Anchor)
+	FString Name;
+
+	/** User friendly description, exposed in client UIs */
+	UPROPERTY(BlueprintReadWrite, Category=Anchor)
+	FString Description;
+
+public:
+	/** Custom implementation of the 'ImportFile' intent */
+	FStringArgDelegate OnImportFileDelegate;
+};
+
+
+/**
+ * Exposes:
+ * - API to enable/disable RemoteImport server (process visibility for clients)
+ * - API to register/unregister RemoteImportAnchors (specific destination for clients)
+ */
+UCLASS(meta=(Experimental))
+class REMOTEIMPORTLIBRARY_API URemoteImportLibrary : public UBlueprintFunctionLibrary
+{
+	GENERATED_BODY()
+
+	URemoteImportLibrary();
+public:
+	/**
+	 * Triggers OnImportFileDelegate on specified Anchor
+	 *
+	 * @param FilePath          Path of the file to import
+	 * @param DestinationName   Name of the destination FRemoteImportAnchor
+	 * @return                  Whether the import operation has been triggered
+	 */
+	UFUNCTION(BlueprintCallable, Category="RemoteImport")
+	static bool ImportSource(const FString& FilePath, const FString& DestinationName);
+
+
+	/**
+	 * Register a possible destination for a remote import
+	 * @param Anchor Description of the destination
+	 */
+	UFUNCTION(BlueprintCallable, Category="RemoteImport")
+	static void RegisterAnchor(const FRemoteImportAnchor& Anchor);
+
+	/**
+	 * Remove a destination
+	 * @param AnchorName Name of the registered anchor to remove
+	 */
+	UFUNCTION(BlueprintCallable, Category="RemoteImport")
+	static void UnregisterAnchor(const FString& AnchorName);
+
+	/**
+	 * List registered Anchor names.
+	 * @return Names that are valid as destination names.
+	 */
+	UFUNCTION(BlueprintCallable, Category="RemoteImport")
+	static TArray<FString> ListAnchors();
+
+
+	/** Enable the server so that the process is visible to clients */
+	UFUNCTION(BlueprintCallable, Exec, Category="RemoteImport")
+	static void StartRemoteImportServer();
+
+	/** Disable the server so that the process is hidden to clients */
+	UFUNCTION(BlueprintCallable, Exec, Category="RemoteImport")
+	static void StopRemoteImportServer();
+
+	/** @return Whether the server is currently active */
+	UFUNCTION(BlueprintCallable, Exec, Category="RemoteImport")
+	static bool IsRemoteImportServerActive();
+
+public:
+	/** @return The delegate that will be triggered when the registered anchor list changes */
+	static FSimpleMulticastDelegate& GetAnchorListChangeDelegate() { return OnAnchorListChange; }
+
+private:
+	static FRemoteImportAnchor* FindAnchor(const FString& AnchorName);
+
+private:
+	static TUniquePtr<FRemoteImportServer> Instance;
+	static TArray<FRemoteImportAnchor> Anchors;
+	static FSimpleMulticastDelegate OnAnchorListChange;
+};
diff --git a/Engine/Plugins/Importers/RemoteImport/Source/RemoteImportLibrary/Public/RemoteImportLibraryModule.h b/Engine/Plugins/Importers/RemoteImport/Source/RemoteImportLibrary/Public/RemoteImportLibraryModule.h
new file mode 100644
--- /dev/null
+++ b/Engine/Plugins/Importers/RemoteImport/Source/RemoteImportLibrary/Public/RemoteImportLibraryModule.h
@@ -0,0 +1,32 @@
+// Copyright Epic Games, Inc. All Rights Reserved.
+
+#pragma once
+
+#include "CoreMinimal.h"
+#include "Modules/ModuleInterface.h"
+#include "Modules/ModuleManager.h"
+
+#define REMOTEIMPORTLIBRARY_MODULE_NAME TEXT("RemoteImportLibrary")
+
+/**
+ * This module exposes base RemoteImport features such as Anchors, and their API.
+ *
+ * RemoteImport is based on an Anchor system, that are setup server-side, discoverable by clients, and act as possible remote import destination.
+ * When a client trigger an import on a destination Anchor, the anchor's Delegate is triggered, allowing any custom implementation to handle the import.
+ *
+ * The client/server communication rely on MessageBus over UDP.
+ * The server is off by default, URemoteImportLibrary provides the API to handle it's life cycle.
+ */
+class FRemoteImportLibraryModule : public IModuleInterface
+{
+public:
+	static inline FRemoteImportLibraryModule& Get()
+	{
+		return FModuleManager::LoadModuleChecked< FRemoteImportLibraryModule >(REMOTEIMPORTLIBRARY_MODULE_NAME);
+	}
+
+	static inline bool IsAvailable()
+	{
+		return FModuleManager::Get().IsModuleLoaded(REMOTEIMPORTLIBRARY_MODULE_NAME);
+	}
+};
diff --git a/Engine/Plugins/Importers/RemoteImport/Source/RemoteImportLibrary/RemoteImportLibrary.Build.cs b/Engine/Plugins/Importers/RemoteImport/Source/RemoteImportLibrary/RemoteImportLibrary.Build.cs
new file mode 100644
--- /dev/null
+++ b/Engine/Plugins/Importers/RemoteImport/Source/RemoteImportLibrary/RemoteImportLibrary.Build.cs
@@ -0,0 +1,17 @@
+// Copyright Epic Games, Inc. All Rights Reserved.
+
+namespace UnrealBuildTool.Rules
+{
+	public class RemoteImportLibrary : ModuleRules
+	{
+		public RemoteImportLibrary(ReadOnlyTargetRules Target) : base(Target)
+		{
+			PrivateDependencyModuleNames.AddRange(new string[]{
+				"Core",
+				"CoreUObject",
+				"Engine",
+				"RemoteImportMessaging",
+			});
+		}
+	}
+}
