Copy of 11164401 from //UE4/Dev-Enterprise:

Datasmith RemoteImport module
Tools to use datasmith translators

diff --git a/Engine/Plugins/Importers/RemoteImport/RemoteImport.uplugin b/Engine/Plugins/Importers/RemoteImport/RemoteImport.uplugin
--- a/Engine/Plugins/Importers/RemoteImport/RemoteImport.uplugin
+++ b/Engine/Plugins/Importers/RemoteImport/RemoteImport.uplugin
@@ -19,6 +19,11 @@
 			"Name": "RemoteImportLibrary",
 			"Type": "Runtime",
 			"LoadingPhase": "Default"
+		},
+		{
+			"Name": "DatasmithRemoteImport",
+			"Type": "Runtime",
+			"LoadingPhase": "Default"
 		}
 	],
 	"Plugins": [
diff --git a/Engine/Plugins/Importers/RemoteImport/Source/DatasmithRemoteImport/DatasmithRemoteImport.Build.cs b/Engine/Plugins/Importers/RemoteImport/Source/DatasmithRemoteImport/DatasmithRemoteImport.Build.cs
new file mode 100644
--- /dev/null
+++ b/Engine/Plugins/Importers/RemoteImport/Source/DatasmithRemoteImport/DatasmithRemoteImport.Build.cs
@@ -0,0 +1,24 @@
+// Copyright Epic Games, Inc. All Rights Reserved.
+
+namespace UnrealBuildTool.Rules
+{
+	public class DatasmithRemoteImport : ModuleRules
+	{
+		public DatasmithRemoteImport(ReadOnlyTargetRules Target) : base(Target)
+		{
+			PublicDependencyModuleNames.AddRange(new string[]{
+				"Core",
+			});
+
+			PrivateDependencyModuleNames.AddRange(new string[]{
+				"CoreUObject",
+				"Engine",
+
+				"RemoteImportLibrary",
+				"MeshDescription",
+				"DatasmithCore",
+				"DatasmithTranslator",
+			});
+		}
+	}
+}
diff --git a/Engine/Plugins/Importers/RemoteImport/Source/DatasmithRemoteImport/Private/DatasmithAdapter.cpp b/Engine/Plugins/Importers/RemoteImport/Source/DatasmithRemoteImport/Private/DatasmithAdapter.cpp
new file mode 100644
--- /dev/null
+++ b/Engine/Plugins/Importers/RemoteImport/Source/DatasmithRemoteImport/Private/DatasmithAdapter.cpp
@@ -0,0 +1,157 @@
+// Copyright Epic Games, Inc. All Rights Reserved.
+
+#include "DatasmithAdapter.h"
+#include "DatasmithRemoteImportLog.h"
+
+#include "DatasmithPayload.h"
+#include "DatasmithSceneFactory.h"
+#include "DatasmithTranslatableSource.h"
+#include "DatasmithTranslator.h"
+#include "IDatasmithSceneElements.h"
+
+#include "Engine/Engine.h"
+#include "Engine/StaticMesh.h"
+#include "Engine/StaticMeshActor.h"
+#include "Engine/World.h"
+#include "MeshDescription.h"
+#include "Templates/SharedPointer.h"
+
+#if WITH_EDITOR
+#include "Editor/EditorEngine.h"
+#include "Editor/UnrealEdEngine.h"
+#endif // WITH_EDITOR
+
+
+namespace DatasmithAdapter
+{
+
+FTranslateResult Translate(const FString& InFilePath)
+{
+	FDatasmithSceneSource Source;
+	Source.SetSourceFile(InFilePath);
+
+	TSharedPtr<FDatasmithTranslatableSceneSource> TranslatableSourcePtr = MakeShared<FDatasmithTranslatableSceneSource>(Source);
+	FDatasmithTranslatableSceneSource& TranslatableSource = *TranslatableSourcePtr;
+
+	if (!TranslatableSource.IsTranslatable())
+	{
+		UE_LOG(LogDatasmithRemoteImport, Error, TEXT("Datasmith adapter import error: no suitable translator found for this source. Abort import."));
+		return {};
+	}
+
+	TSharedRef<IDatasmithScene> Scene = FDatasmithSceneFactory::CreateScene(*Source.GetSceneName());
+
+	if (!TranslatableSource.Translate(Scene))
+	{
+		UE_LOG(LogDatasmithRemoteImport, Error, TEXT("Datasmith import error: Scene translation failure. Abort import."));
+		return {};
+	}
+
+	return { Scene, TranslatableSource.GetTranslator(), TranslatableSourcePtr};
+}
+
+
+void ListActorElements(IDatasmithActorElement* Actor, TArray<IDatasmithActorElement*>& OutElements)
+{
+	for (int32 ChildElementIndex = 0 ; ChildElementIndex < Actor->GetChildrenCount(); ++ChildElementIndex)
+	{
+		const TSharedPtr<IDatasmithActorElement>& ActorElement = Actor->GetChild(ChildElementIndex);
+		ListActorElements(ActorElement.Get(), OutElements);
+		OutElements.Add(ActorElement.Get());
+	}
+}
+
+
+void ListActorElements(IDatasmithScene* Scene, TArray<IDatasmithActorElement*>& OutElements)
+{
+	for (int32 ActorElementIndex = 0 ; ActorElementIndex < Scene->GetActorsCount(); ++ActorElementIndex)
+	{
+		const TSharedPtr<IDatasmithActorElement>& ActorElement = Scene->GetActor(ActorElementIndex);
+		ListActorElements(ActorElement.Get(), OutElements);
+		OutElements.Add(ActorElement.Get());
+	}
+}
+
+
+bool Import(FTranslateResult& Translation, UWorld* TargetWorld, const FTransform& RootTM)
+{
+	IDatasmithScene* Scene = Translation.Scene.Get();
+	IDatasmithTranslator* Translator = Translation.Translator.Get();
+
+	if (Scene == nullptr || Translator == nullptr || TargetWorld == nullptr)
+	{
+		UE_LOG(LogDatasmithRemoteImport, Error, TEXT("Can't import: incomplete input"));
+		return false;
+	}
+
+	struct FMeshPath
+	{
+		const TCHAR* Name;
+		TSharedRef<IDatasmithMeshElement> Element;
+	};
+	TMap<FString, FMeshPath> KnownMeshInfo;
+
+	for (int32 i = 0 ; i < Scene->GetMeshesCount(); ++i)
+	{
+		TSharedRef<IDatasmithMeshElement> MeshElement = Scene->GetMesh(i).ToSharedRef();
+		const TCHAR* MeshName = MeshElement->GetName();
+		KnownMeshInfo.Add(MeshName, {MeshName, MeshElement});
+	}
+
+	TArray<IDatasmithActorElement*> AllMeshActors;
+	ListActorElements(Scene, AllMeshActors);
+
+
+	for (IDatasmithActorElement* ActorElement : AllMeshActors)
+	{
+		if (ActorElement && ActorElement->IsA(EDatasmithElementType::StaticMeshActor))
+		{
+			IDatasmithMeshActorElement* MeshActorElement = static_cast<IDatasmithMeshActorElement*>(ActorElement);
+
+			const TCHAR* MeshName = MeshActorElement->GetStaticMeshPathName();
+
+			if (FMeshPath* ThisMeshInfo = KnownMeshInfo.Find(MeshName))
+			{
+				FDatasmithMeshElementPayload MeshPayload;
+				Translator->LoadStaticMesh(ThisMeshInfo->Element, MeshPayload);
+
+				if (MeshPayload.LodMeshes.IsValidIndex(0))
+				{
+					FTransform WorldTransform (
+						ActorElement->GetRotation(),
+						ActorElement->GetTranslation(),
+						ActorElement->GetScale()
+					);
+
+					AStaticMeshActor* NewMeshActor = TargetWorld->SpawnActor<AStaticMeshActor>(AStaticMeshActor::StaticClass(), WorldTransform*RootTM);
+#if WITH_EDITOR
+					NewMeshActor->SetActorLabel(MeshActorElement->GetLabel());
+#endif // WITH_EDITOR
+
+					if (NewMeshActor == nullptr)
+					{
+						UE_LOG(LogDatasmithRemoteImport, Error, TEXT("Error on static mesh actor spawn"));
+						return false;
+					}
+
+					if (UStaticMesh* StaticMeshObject = NewObject<UStaticMesh>())
+					{
+						TArray<const FMeshDescription*> MeshDescriptions;
+						for (auto& MeshDescription: MeshPayload.LodMeshes)
+						{
+							MeshDescriptions.Add(&MeshDescription);
+						}
+
+						StaticMeshObject->BuildFromMeshDescriptions(MeshDescriptions);
+						NewMeshActor->GetStaticMeshComponent()->SetMobility(EComponentMobility::Movable);
+						NewMeshActor->GetStaticMeshComponent()->SetStaticMesh(StaticMeshObject);
+					}
+				}
+			}
+		}
+	}
+	return true;
+}
+
+} // ns DatasmithAdapter
+
diff --git a/Engine/Plugins/Importers/RemoteImport/Source/DatasmithRemoteImport/Private/DatasmithRemoteEndpointComponent.cpp b/Engine/Plugins/Importers/RemoteImport/Source/DatasmithRemoteImport/Private/DatasmithRemoteEndpointComponent.cpp
new file mode 100644
--- /dev/null
+++ b/Engine/Plugins/Importers/RemoteImport/Source/DatasmithRemoteImport/Private/DatasmithRemoteEndpointComponent.cpp
@@ -0,0 +1,35 @@
+// Copyright Epic Games, Inc. All Rights Reserved.
+
+#include "DatasmithRemoteEndpointComponent.h"
+#include "DatasmithAdapter.h"
+#include "DatasmithRemoteImportLog.h"
+#include "RemoteImportLibrary.h"
+
+
+void UDatasmithRemoteEndpointComponent::BeginPlay()
+{
+	Super::BeginPlay();
+	FRemoteImportAnchor Anchor;
+	RegisteredAnchorName = Name.IsEmpty() ? Name = GetName() : Name;
+	Anchor.Name = RegisteredAnchorName;
+	Anchor.Description = TEXT("UDatasmithRemoteEndpointComponent");
+	Anchor.OnImportFileDelegate.BindUObject(this, &UDatasmithRemoteEndpointComponent::OnImportFile);
+
+	URemoteImportLibrary::RegisterAnchor(Anchor);
+}
+
+void UDatasmithRemoteEndpointComponent::EndPlay(const EEndPlayReason::Type EndPlayReason)
+{
+	Super::EndPlay(EndPlayReason);
+	URemoteImportLibrary::UnregisterAnchor(RegisteredAnchorName);
+}
+
+void UDatasmithRemoteEndpointComponent::OnImportFile(const FString& FilePath)
+{
+	UE_LOG(LogDatasmithRemoteImport, Display, TEXT("Translate %s..."), *FilePath);
+	DatasmithAdapter::FTranslateResult Content = DatasmithAdapter::Translate(FilePath);
+	UE_LOG(LogDatasmithRemoteImport, Display, TEXT("Import %s..."), *FilePath);
+	bool ImportResult = DatasmithAdapter::Import(Content, GetWorld(), GetComponentTransform());
+	UE_LOG(LogDatasmithRemoteImport, Display, TEXT("Import done for %s"), *FilePath);
+}
+
diff --git a/Engine/Plugins/Importers/RemoteImport/Source/DatasmithRemoteImport/Private/DatasmithRemoteImportLog.h b/Engine/Plugins/Importers/RemoteImport/Source/DatasmithRemoteImport/Private/DatasmithRemoteImportLog.h
new file mode 100644
--- /dev/null
+++ b/Engine/Plugins/Importers/RemoteImport/Source/DatasmithRemoteImport/Private/DatasmithRemoteImportLog.h
@@ -0,0 +1,7 @@
+// Copyright Epic Games, Inc. All Rights Reserved.
+
+#pragma once
+
+#include "Logging/LogMacros.h"
+
+DECLARE_LOG_CATEGORY_EXTERN(LogDatasmithRemoteImport, Log, All);
diff --git a/Engine/Plugins/Importers/RemoteImport/Source/DatasmithRemoteImport/Private/DatasmithRemoteImportModule.cpp b/Engine/Plugins/Importers/RemoteImport/Source/DatasmithRemoteImport/Private/DatasmithRemoteImportModule.cpp
new file mode 100644
--- /dev/null
+++ b/Engine/Plugins/Importers/RemoteImport/Source/DatasmithRemoteImport/Private/DatasmithRemoteImportModule.cpp
@@ -0,0 +1,9 @@
+// Copyright Epic Games, Inc. All Rights Reserved.
+#include "DatasmithRemoteImportModule.h"
+#include "DatasmithRemoteImportLog.h"
+
+#include "Modules/ModuleManager.h"
+
+IMPLEMENT_MODULE(FDatasmithRemoteImportModule, DatasmithRemoteImport);
+
+DEFINE_LOG_CATEGORY(LogDatasmithRemoteImport)
\ No newline at end of file
diff --git a/Engine/Plugins/Importers/RemoteImport/Source/DatasmithRemoteImport/Public/DatasmithAdapter.h b/Engine/Plugins/Importers/RemoteImport/Source/DatasmithRemoteImport/Public/DatasmithAdapter.h
new file mode 100644
--- /dev/null
+++ b/Engine/Plugins/Importers/RemoteImport/Source/DatasmithRemoteImport/Public/DatasmithAdapter.h
@@ -0,0 +1,31 @@
+// Copyright Epic Games, Inc. All Rights Reserved.
+
+#pragma once
+
+
+#include "CoreMinimal.h"
+#include "IDatasmithSceneElements.h"
+
+class IDatasmithTranslator;
+struct FDatasmithTranslatableSceneSource;
+class UWorld;
+struct FRemoteImportAnchor;
+
+
+namespace DatasmithAdapter
+{
+	/**	Hold objects required to exploit a datasmith scene */
+	struct DATASMITHREMOTEIMPORT_API FTranslateResult
+	{
+		TSharedPtr<IDatasmithScene> Scene;
+		TSharedPtr<IDatasmithTranslator> Translator;
+		TSharedPtr<FDatasmithTranslatableSceneSource> SceneScope;
+	};
+
+	/** Translate a source file into a DatasmithScene representations and it's Translator */
+	DATASMITHREMOTEIMPORT_API FTranslateResult Translate(const FString& InFilePath);
+
+	/** Proof of concept import: loads a DatasmithScene geometry into a World */
+	DATASMITHREMOTEIMPORT_API bool Import(FTranslateResult& Translation, UWorld* TargetWorld, const FTransform& RootTransform={});
+
+} // ns DatasmithAdapter
diff --git a/Engine/Plugins/Importers/RemoteImport/Source/DatasmithRemoteImport/Public/DatasmithRemoteEndpointComponent.h b/Engine/Plugins/Importers/RemoteImport/Source/DatasmithRemoteImport/Public/DatasmithRemoteEndpointComponent.h
new file mode 100644
--- /dev/null
+++ b/Engine/Plugins/Importers/RemoteImport/Source/DatasmithRemoteImport/Public/DatasmithRemoteEndpointComponent.h
@@ -0,0 +1,30 @@
+// Copyright Epic Games, Inc. All Rights Reserved.
+
+#include "CoreMinimal.h"
+#include "Components/SceneComponent.h"
+
+#include "DatasmithRemoteEndpointComponent.generated.h"
+
+/**
+ * Automatically register/unregister an RemoteImportAnchor while in play
+ * Accept a RemoteImport command to load a file by path
+ */
+UCLASS(BlueprintType, meta=(BlueprintSpawnableComponent, Experimental))
+class DATASMITHREMOTEIMPORT_API UDatasmithRemoteEndpointComponent : public USceneComponent
+{
+	GENERATED_BODY()
+
+public:
+	virtual void BeginPlay() override;
+	virtual void EndPlay(const EEndPlayReason::Type EndPlayReason) override;
+
+public:
+	/** Name that will be exposed to the RemoteImport system */
+	UPROPERTY(BlueprintReadWrite, EditAnywhere, Category=Anchor)
+	FString Name;
+
+private:
+	FString RegisteredAnchorName; // Snapshot registered name for correct unregister
+	void OnImportFile(const FString& FilePath);
+};
+
diff --git a/Engine/Plugins/Importers/RemoteImport/Source/DatasmithRemoteImport/Public/DatasmithRemoteImportModule.h b/Engine/Plugins/Importers/RemoteImport/Source/DatasmithRemoteImport/Public/DatasmithRemoteImportModule.h
new file mode 100644
--- /dev/null
+++ b/Engine/Plugins/Importers/RemoteImport/Source/DatasmithRemoteImport/Public/DatasmithRemoteImportModule.h
@@ -0,0 +1,28 @@
+// Copyright Epic Games, Inc. All Rights Reserved.
+
+#pragma once
+
+#include "CoreMinimal.h"
+#include "Modules/ModuleInterface.h"
+#include "Modules/ModuleManager.h"
+
+#define DATASMITHREMOTEIMPORT_MODULE_NAME TEXT("DatasmithRemoteImport")
+
+/**
+ * Exposes:
+ * - DatasmithAdapter                   Proxy for Datasmith Translators
+ * - UDatasmithRemoteEndpointComponent  SceneComponent that scopes a Remote Import Anchor while in Play
+ */
+class DATASMITHREMOTEIMPORT_API FDatasmithRemoteImportModule : public IModuleInterface
+{
+public:
+    static inline FDatasmithRemoteImportModule& Get()
+    {
+        return FModuleManager::LoadModuleChecked< FDatasmithRemoteImportModule >(DATASMITHREMOTEIMPORT_MODULE_NAME);
+    }
+
+    static inline bool IsAvailable()
+    {
+        return FModuleManager::Get().IsModuleLoaded(DATASMITHREMOTEIMPORT_MODULE_NAME);
+    }
+};
