Copy of 11195103 from //UE4/Dev-Enterprise:

Fix missing header guard

diff --git a/Engine/Plugins/Importers/RemoteImport/Source/DatasmithRemoteImport/Public/DatasmithRemoteEndpointComponent.h b/Engine/Plugins/Importers/RemoteImport/Source/DatasmithRemoteImport/Public/DatasmithRemoteEndpointComponent.h
--- a/Engine/Plugins/Importers/RemoteImport/Source/DatasmithRemoteImport/Public/DatasmithRemoteEndpointComponent.h
+++ b/Engine/Plugins/Importers/RemoteImport/Source/DatasmithRemoteImport/Public/DatasmithRemoteEndpointComponent.h
@@ -1,5 +1,7 @@
 // Copyright Epic Games, Inc. All Rights Reserved.
 
+#pragma once
+
 #include "CoreMinimal.h"
 #include "Components/SceneComponent.h"
 
