Copy of 11221350 from //UE4/Dev-Enterprise:

RemoteImportLibrary: Prevent duplicated Anchor names

diff --git a/Engine/Plugins/Importers/RemoteImport/Source/DatasmithRemoteImport/Private/DatasmithRemoteEndpointComponent.cpp b/Engine/Plugins/Importers/RemoteImport/Source/DatasmithRemoteImport/Private/DatasmithRemoteEndpointComponent.cpp
--- a/Engine/Plugins/Importers/RemoteImport/Source/DatasmithRemoteImport/Private/DatasmithRemoteEndpointComponent.cpp
+++ b/Engine/Plugins/Importers/RemoteImport/Source/DatasmithRemoteImport/Private/DatasmithRemoteEndpointComponent.cpp
@@ -10,12 +10,11 @@ void UDatasmithRemoteEndpointComponent::BeginPlay()
 {
 	Super::BeginPlay();
 	FRemoteImportAnchor Anchor;
-	RegisteredAnchorName = Name.IsEmpty() ? Name = GetName() : Name;
-	Anchor.Name = RegisteredAnchorName;
+	Anchor.Name = Name.IsEmpty() ? GetName() : Name;
 	Anchor.Description = TEXT("UDatasmithRemoteEndpointComponent");
 	Anchor.OnImportFileDelegate.BindUObject(this, &UDatasmithRemoteEndpointComponent::OnImportFile);
 
-	URemoteImportLibrary::RegisterAnchor(Anchor);
+	RegisteredAnchorName = URemoteImportLibrary::RegisterAnchor(Anchor);
 }
 
 void UDatasmithRemoteEndpointComponent::EndPlay(const EEndPlayReason::Type EndPlayReason)
diff --git a/Engine/Plugins/Importers/RemoteImport/Source/RemoteImportLibrary/Private/RemoteImportLibrary.cpp b/Engine/Plugins/Importers/RemoteImport/Source/RemoteImportLibrary/Private/RemoteImportLibrary.cpp
--- a/Engine/Plugins/Importers/RemoteImport/Source/RemoteImportLibrary/Private/RemoteImportLibrary.cpp
+++ b/Engine/Plugins/Importers/RemoteImport/Source/RemoteImportLibrary/Private/RemoteImportLibrary.cpp
@@ -34,20 +34,32 @@ bool URemoteImportLibrary::ImportSource(const FString& FilePath, const FString&
 }
 
 
-void URemoteImportLibrary::RegisterAnchor(const FRemoteImportAnchor& Anchor)
+FString URemoteImportLibrary::RegisterAnchor(const FRemoteImportAnchor& Anchor, bool bAllowRename)
 {
-	Anchors.Add(Anchor);
+	FRemoteImportAnchor AnchorCopy = Anchor;
+	FString InitialName = AnchorCopy.Name;
+	int32 Suffix = 0;
+	while (Anchors.FindByPredicate([&](const FRemoteImportAnchor& Registered){return Registered.Name == AnchorCopy.Name;}))
+	{
+		if (!bAllowRename)
+		{
+			return {};
+		}
+		AnchorCopy.Name = InitialName + TEXT("_") + FString::FromInt(Suffix++);
+	}
+	Anchors.Add(AnchorCopy);
 	OnAnchorListChange.Broadcast();
+	return AnchorCopy.Name;
 }
 
-
-void URemoteImportLibrary::UnregisterAnchor(const FString& AnchorName)
+void URemoteImportLibrary::UnregisterAnchor(const FString& AnchorHandle)
 {
-	Anchors.RemoveAll([&AnchorName](const FRemoteImportAnchor& Anchor){ return Anchor.Name == AnchorName; });
-	OnAnchorListChange.Broadcast();
+	if (Anchors.RemoveAll([&AnchorHandle](const FRemoteImportAnchor& Anchor){ return Anchor.Name == AnchorHandle; }))
+	{
+		OnAnchorListChange.Broadcast();
+	}
 }
 
-
 TArray<FString> URemoteImportLibrary::ListAnchors()
 {
 	TArray<FString> Out;
diff --git a/Engine/Plugins/Importers/RemoteImport/Source/RemoteImportLibrary/Public/RemoteImportLibrary.h b/Engine/Plugins/Importers/RemoteImport/Source/RemoteImportLibrary/Public/RemoteImportLibrary.h
--- a/Engine/Plugins/Importers/RemoteImport/Source/RemoteImportLibrary/Public/RemoteImportLibrary.h
+++ b/Engine/Plugins/Importers/RemoteImport/Source/RemoteImportLibrary/Public/RemoteImportLibrary.h
@@ -60,42 +60,44 @@ public:
 	 * @param DestinationName   Name of the destination FRemoteImportAnchor
 	 * @return                  Whether the import operation has been triggered
 	 */
-	UFUNCTION(BlueprintCallable, Category="RemoteImport")
+	UFUNCTION(BlueprintCallable, Category="RemoteImport|Action")
 	static bool ImportSource(const FString& FilePath, const FString& DestinationName);
 
 
 	/**
 	 * Register a possible destination for a remote import
-	 * @param Anchor Description of the destination
+	 * @param Anchor        Description of the destination
+	 * @param bAllowRename  When true, a unique name is generated to register the anchor even when names collide
+	 * @returns the anchor handle, required to unregister the Anchor
 	 */
-	UFUNCTION(BlueprintCallable, Category="RemoteImport")
-	static void RegisterAnchor(const FRemoteImportAnchor& Anchor);
+	UFUNCTION(BlueprintCallable, Category="RemoteImport|Anchor")
+	static FString RegisterAnchor(const FRemoteImportAnchor& Anchor, bool bAllowRename=true);
 
 	/**
 	 * Remove a destination
-	 * @param AnchorName Name of the registered anchor to remove
+	 * @param AnchorHandle  Registered handle as returned by RegisterAnchor()
 	 */
-	UFUNCTION(BlueprintCallable, Category="RemoteImport")
-	static void UnregisterAnchor(const FString& AnchorName);
+	UFUNCTION(BlueprintCallable, Category="RemoteImport|Anchor")
+	static void UnregisterAnchor(const FString& AnchorHandle);
 
 	/**
 	 * List registered Anchor names.
 	 * @return Names that are valid as destination names.
 	 */
-	UFUNCTION(BlueprintCallable, Category="RemoteImport")
+	UFUNCTION(BlueprintCallable, Category="RemoteImport|Anchor")
 	static TArray<FString> ListAnchors();
 
 
 	/** Enable the server so that the process is visible to clients */
-	UFUNCTION(BlueprintCallable, Exec, Category="RemoteImport")
+	UFUNCTION(BlueprintCallable, Category="RemoteImport|Server")
 	static void StartRemoteImportServer();
 
 	/** Disable the server so that the process is hidden to clients */
-	UFUNCTION(BlueprintCallable, Exec, Category="RemoteImport")
+	UFUNCTION(BlueprintCallable, Category="RemoteImport|Server")
 	static void StopRemoteImportServer();
 
 	/** @return Whether the server is currently active */
-	UFUNCTION(BlueprintCallable, Exec, Category="RemoteImport")
+	UFUNCTION(BlueprintCallable, Category="RemoteImport|Server")
 	static bool IsRemoteImportServerActive();
 
 public:
