moved /Engine/Plugins/Importers/RemoteImport to /Engine/Plugins/Importers/Experimental/RemoteImport

diff --git a/Engine/Plugins/Importers/Experimental/RemoteImport/RemoteImport.uplugin b/Engine/Plugins/Importers/Experimental/RemoteImport/RemoteImport.uplugin
new file mode 100644
--- /dev/null
+++ b/Engine/Plugins/Importers/Experimental/RemoteImport/RemoteImport.uplugin
@@ -0,0 +1,39 @@
+{
+	"FileVersion": 3,
+	"Version": 1,
+	"VersionName": "1.0",
+	"FriendlyName": "Remote Import",
+	"Description": "Allows external applications to remotely trigger a content import",
+	"Category": "Importers",
+	"CreatedBy": "Epic Games, Inc.",
+	"CreatedByURL": "http://epicgames.com",
+	"DocsURL": "",
+	"MarketplaceURL": "",
+	"SupportURL": "",
+	"EnabledByDefault": false,
+	"CanContainContent": false,
+	"IsBetaVersion": true,
+	"Installed": false,
+	"Modules": [
+		{
+			"Name": "RemoteImportLibrary",
+			"Type": "Runtime",
+			"LoadingPhase": "Default"
+		},
+		{
+			"Name": "DatasmithRemoteImport",
+			"Type": "Runtime",
+			"LoadingPhase": "Default"
+		}
+	],
+	"Plugins": [
+		{
+			"Name": "OnlineSubsystem",
+			"Enabled": true
+		},
+		{
+			"Name": "DatasmithImporter",
+			"Enabled": true
+		}
+	]
+}
\ No newline at end of file
diff --git a/Engine/Plugins/Importers/Experimental/RemoteImport/Source/DatasmithRemoteImport/DatasmithRemoteImport.Build.cs b/Engine/Plugins/Importers/Experimental/RemoteImport/Source/DatasmithRemoteImport/DatasmithRemoteImport.Build.cs
new file mode 100644
--- /dev/null
+++ b/Engine/Plugins/Importers/Experimental/RemoteImport/Source/DatasmithRemoteImport/DatasmithRemoteImport.Build.cs
@@ -0,0 +1,24 @@
+// Copyright Epic Games, Inc. All Rights Reserved.
+
+namespace UnrealBuildTool.Rules
+{
+	public class DatasmithRemoteImport : ModuleRules
+	{
+		public DatasmithRemoteImport(ReadOnlyTargetRules Target) : base(Target)
+		{
+			PublicDependencyModuleNames.AddRange(new string[]{
+				"Core",
+			});
+
+			PrivateDependencyModuleNames.AddRange(new string[]{
+				"CoreUObject",
+				"Engine",
+
+				"RemoteImportLibrary",
+				"MeshDescription",
+				"DatasmithCore",
+				"DatasmithTranslator",
+			});
+		}
+	}
+}
diff --git a/Engine/Plugins/Importers/Experimental/RemoteImport/Source/DatasmithRemoteImport/Private/DatasmithAdapter.cpp b/Engine/Plugins/Importers/Experimental/RemoteImport/Source/DatasmithRemoteImport/Private/DatasmithAdapter.cpp
new file mode 100644
--- /dev/null
+++ b/Engine/Plugins/Importers/Experimental/RemoteImport/Source/DatasmithRemoteImport/Private/DatasmithAdapter.cpp
@@ -0,0 +1,157 @@
+// Copyright Epic Games, Inc. All Rights Reserved.
+
+#include "DatasmithAdapter.h"
+#include "DatasmithRemoteImportLog.h"
+
+#include "DatasmithPayload.h"
+#include "DatasmithSceneFactory.h"
+#include "DatasmithTranslatableSource.h"
+#include "DatasmithTranslator.h"
+#include "IDatasmithSceneElements.h"
+
+#include "Engine/Engine.h"
+#include "Engine/StaticMesh.h"
+#include "Engine/StaticMeshActor.h"
+#include "Engine/World.h"
+#include "MeshDescription.h"
+#include "Templates/SharedPointer.h"
+
+#if WITH_EDITOR
+#include "Editor/EditorEngine.h"
+#include "Editor/UnrealEdEngine.h"
+#endif // WITH_EDITOR
+
+
+namespace DatasmithAdapter
+{
+
+FTranslateResult Translate(const FString& InFilePath)
+{
+	FDatasmithSceneSource Source;
+	Source.SetSourceFile(InFilePath);
+
+	TSharedPtr<FDatasmithTranslatableSceneSource> TranslatableSourcePtr = MakeShared<FDatasmithTranslatableSceneSource>(Source);
+	FDatasmithTranslatableSceneSource& TranslatableSource = *TranslatableSourcePtr;
+
+	if (!TranslatableSource.IsTranslatable())
+	{
+		UE_LOG(LogDatasmithRemoteImport, Error, TEXT("Datasmith adapter import error: no suitable translator found for this source. Abort import."));
+		return {};
+	}
+
+	TSharedRef<IDatasmithScene> Scene = FDatasmithSceneFactory::CreateScene(*Source.GetSceneName());
+
+	if (!TranslatableSource.Translate(Scene))
+	{
+		UE_LOG(LogDatasmithRemoteImport, Error, TEXT("Datasmith import error: Scene translation failure. Abort import."));
+		return {};
+	}
+
+	return { Scene, TranslatableSource.GetTranslator(), TranslatableSourcePtr};
+}
+
+
+void ListActorElements(IDatasmithActorElement* Actor, TArray<IDatasmithActorElement*>& OutElements)
+{
+	for (int32 ChildElementIndex = 0 ; ChildElementIndex < Actor->GetChildrenCount(); ++ChildElementIndex)
+	{
+		const TSharedPtr<IDatasmithActorElement>& ActorElement = Actor->GetChild(ChildElementIndex);
+		ListActorElements(ActorElement.Get(), OutElements);
+		OutElements.Add(ActorElement.Get());
+	}
+}
+
+
+void ListActorElements(IDatasmithScene* Scene, TArray<IDatasmithActorElement*>& OutElements)
+{
+	for (int32 ActorElementIndex = 0 ; ActorElementIndex < Scene->GetActorsCount(); ++ActorElementIndex)
+	{
+		const TSharedPtr<IDatasmithActorElement>& ActorElement = Scene->GetActor(ActorElementIndex);
+		ListActorElements(ActorElement.Get(), OutElements);
+		OutElements.Add(ActorElement.Get());
+	}
+}
+
+
+bool Import(FTranslateResult& Translation, UWorld* TargetWorld, const FTransform& RootTM)
+{
+	IDatasmithScene* Scene = Translation.Scene.Get();
+	IDatasmithTranslator* Translator = Translation.Translator.Get();
+
+	if (Scene == nullptr || Translator == nullptr || TargetWorld == nullptr)
+	{
+		UE_LOG(LogDatasmithRemoteImport, Error, TEXT("Can't import: incomplete input"));
+		return false;
+	}
+
+	struct FMeshPath
+	{
+		const TCHAR* Name;
+		TSharedRef<IDatasmithMeshElement> Element;
+	};
+	TMap<FString, FMeshPath> KnownMeshInfo;
+
+	for (int32 i = 0 ; i < Scene->GetMeshesCount(); ++i)
+	{
+		TSharedRef<IDatasmithMeshElement> MeshElement = Scene->GetMesh(i).ToSharedRef();
+		const TCHAR* MeshName = MeshElement->GetName();
+		KnownMeshInfo.Add(MeshName, {MeshName, MeshElement});
+	}
+
+	TArray<IDatasmithActorElement*> AllMeshActors;
+	ListActorElements(Scene, AllMeshActors);
+
+
+	for (IDatasmithActorElement* ActorElement : AllMeshActors)
+	{
+		if (ActorElement && ActorElement->IsA(EDatasmithElementType::StaticMeshActor))
+		{
+			IDatasmithMeshActorElement* MeshActorElement = static_cast<IDatasmithMeshActorElement*>(ActorElement);
+
+			const TCHAR* MeshName = MeshActorElement->GetStaticMeshPathName();
+
+			if (FMeshPath* ThisMeshInfo = KnownMeshInfo.Find(MeshName))
+			{
+				FDatasmithMeshElementPayload MeshPayload;
+				Translator->LoadStaticMesh(ThisMeshInfo->Element, MeshPayload);
+
+				if (MeshPayload.LodMeshes.IsValidIndex(0))
+				{
+					FTransform WorldTransform (
+						ActorElement->GetRotation(),
+						ActorElement->GetTranslation(),
+						ActorElement->GetScale()
+					);
+
+					AStaticMeshActor* NewMeshActor = TargetWorld->SpawnActor<AStaticMeshActor>(AStaticMeshActor::StaticClass(), WorldTransform*RootTM);
+					if (NewMeshActor == nullptr)
+					{
+						UE_LOG(LogDatasmithRemoteImport, Error, TEXT("Error on static mesh actor spawn"));
+						return false;
+					}
+
+#if WITH_EDITOR
+					NewMeshActor->SetActorLabel(MeshActorElement->GetLabel());
+#endif // WITH_EDITOR
+
+					if (UStaticMesh* StaticMeshObject = NewObject<UStaticMesh>())
+					{
+						TArray<const FMeshDescription*> MeshDescriptions;
+						for (auto& MeshDescription: MeshPayload.LodMeshes)
+						{
+							MeshDescriptions.Add(&MeshDescription);
+						}
+
+						StaticMeshObject->BuildFromMeshDescriptions(MeshDescriptions);
+						NewMeshActor->GetStaticMeshComponent()->SetMobility(EComponentMobility::Movable);
+						NewMeshActor->GetStaticMeshComponent()->SetStaticMesh(StaticMeshObject);
+					}
+				}
+			}
+		}
+	}
+	return true;
+}
+
+} // ns DatasmithAdapter
+
diff --git a/Engine/Plugins/Importers/Experimental/RemoteImport/Source/DatasmithRemoteImport/Private/DatasmithRemoteEndpointComponent.cpp b/Engine/Plugins/Importers/Experimental/RemoteImport/Source/DatasmithRemoteImport/Private/DatasmithRemoteEndpointComponent.cpp
new file mode 100644
--- /dev/null
+++ b/Engine/Plugins/Importers/Experimental/RemoteImport/Source/DatasmithRemoteImport/Private/DatasmithRemoteEndpointComponent.cpp
@@ -0,0 +1,34 @@
+// Copyright Epic Games, Inc. All Rights Reserved.
+
+#include "DatasmithRemoteEndpointComponent.h"
+#include "DatasmithAdapter.h"
+#include "DatasmithRemoteImportLog.h"
+#include "RemoteImportLibrary.h"
+
+
+void UDatasmithRemoteEndpointComponent::BeginPlay()
+{
+	Super::BeginPlay();
+	FRemoteImportAnchor Anchor;
+	Anchor.Name = Name.IsEmpty() ? GetName() : Name;
+	Anchor.Description = TEXT("UDatasmithRemoteEndpointComponent");
+	Anchor.OnImportFileDelegate.BindUObject(this, &UDatasmithRemoteEndpointComponent::OnImportFile);
+
+	RegisteredAnchorName = URemoteImportLibrary::RegisterAnchor(Anchor);
+}
+
+void UDatasmithRemoteEndpointComponent::EndPlay(const EEndPlayReason::Type EndPlayReason)
+{
+	Super::EndPlay(EndPlayReason);
+	URemoteImportLibrary::UnregisterAnchor(RegisteredAnchorName);
+}
+
+void UDatasmithRemoteEndpointComponent::OnImportFile(const FString& FilePath)
+{
+	UE_LOG(LogDatasmithRemoteImport, Display, TEXT("Translate %s..."), *FilePath);
+	DatasmithAdapter::FTranslateResult Content = DatasmithAdapter::Translate(FilePath);
+	UE_LOG(LogDatasmithRemoteImport, Display, TEXT("Import %s..."), *FilePath);
+	bool ImportResult = DatasmithAdapter::Import(Content, GetWorld(), GetComponentTransform());
+	UE_LOG(LogDatasmithRemoteImport, Display, TEXT("Import done for %s"), *FilePath);
+}
+
diff --git a/Engine/Plugins/Importers/Experimental/RemoteImport/Source/DatasmithRemoteImport/Private/DatasmithRemoteImportLog.h b/Engine/Plugins/Importers/Experimental/RemoteImport/Source/DatasmithRemoteImport/Private/DatasmithRemoteImportLog.h
new file mode 100644
--- /dev/null
+++ b/Engine/Plugins/Importers/Experimental/RemoteImport/Source/DatasmithRemoteImport/Private/DatasmithRemoteImportLog.h
@@ -0,0 +1,7 @@
+// Copyright Epic Games, Inc. All Rights Reserved.
+
+#pragma once
+
+#include "Logging/LogMacros.h"
+
+DECLARE_LOG_CATEGORY_EXTERN(LogDatasmithRemoteImport, Log, All);
diff --git a/Engine/Plugins/Importers/Experimental/RemoteImport/Source/DatasmithRemoteImport/Private/DatasmithRemoteImportModule.cpp b/Engine/Plugins/Importers/Experimental/RemoteImport/Source/DatasmithRemoteImport/Private/DatasmithRemoteImportModule.cpp
new file mode 100644
--- /dev/null
+++ b/Engine/Plugins/Importers/Experimental/RemoteImport/Source/DatasmithRemoteImport/Private/DatasmithRemoteImportModule.cpp
@@ -0,0 +1,9 @@
+// Copyright Epic Games, Inc. All Rights Reserved.
+#include "DatasmithRemoteImportModule.h"
+#include "DatasmithRemoteImportLog.h"
+
+#include "Modules/ModuleManager.h"
+
+IMPLEMENT_MODULE(FDatasmithRemoteImportModule, DatasmithRemoteImport);
+
+DEFINE_LOG_CATEGORY(LogDatasmithRemoteImport)
\ No newline at end of file
diff --git a/Engine/Plugins/Importers/Experimental/RemoteImport/Source/DatasmithRemoteImport/Public/DatasmithAdapter.h b/Engine/Plugins/Importers/Experimental/RemoteImport/Source/DatasmithRemoteImport/Public/DatasmithAdapter.h
new file mode 100644
--- /dev/null
+++ b/Engine/Plugins/Importers/Experimental/RemoteImport/Source/DatasmithRemoteImport/Public/DatasmithAdapter.h
@@ -0,0 +1,31 @@
+// Copyright Epic Games, Inc. All Rights Reserved.
+
+#pragma once
+
+
+#include "CoreMinimal.h"
+#include "IDatasmithSceneElements.h"
+
+class IDatasmithTranslator;
+struct FDatasmithTranslatableSceneSource;
+class UWorld;
+struct FRemoteImportAnchor;
+
+
+namespace DatasmithAdapter
+{
+	/**	Hold objects required to exploit a datasmith scene */
+	struct DATASMITHREMOTEIMPORT_API FTranslateResult
+	{
+		TSharedPtr<IDatasmithScene> Scene;
+		TSharedPtr<IDatasmithTranslator> Translator;
+		TSharedPtr<FDatasmithTranslatableSceneSource> SceneScope;
+	};
+
+	/** Translate a source file into a DatasmithScene representations and it's Translator */
+	DATASMITHREMOTEIMPORT_API FTranslateResult Translate(const FString& InFilePath);
+
+	/** Proof of concept import: loads a DatasmithScene geometry into a World */
+	DATASMITHREMOTEIMPORT_API bool Import(FTranslateResult& Translation, UWorld* TargetWorld, const FTransform& RootTransform={});
+
+} // ns DatasmithAdapter
diff --git a/Engine/Plugins/Importers/Experimental/RemoteImport/Source/DatasmithRemoteImport/Public/DatasmithRemoteEndpointComponent.h b/Engine/Plugins/Importers/Experimental/RemoteImport/Source/DatasmithRemoteImport/Public/DatasmithRemoteEndpointComponent.h
new file mode 100644
--- /dev/null
+++ b/Engine/Plugins/Importers/Experimental/RemoteImport/Source/DatasmithRemoteImport/Public/DatasmithRemoteEndpointComponent.h
@@ -0,0 +1,32 @@
+// Copyright Epic Games, Inc. All Rights Reserved.
+
+#pragma once
+
+#include "CoreMinimal.h"
+#include "Components/SceneComponent.h"
+
+#include "DatasmithRemoteEndpointComponent.generated.h"
+
+/**
+ * Automatically register/unregister an RemoteImportAnchor while in play
+ * Accept a RemoteImport command to load a file by path
+ */
+UCLASS(BlueprintType, meta=(BlueprintSpawnableComponent, Experimental))
+class DATASMITHREMOTEIMPORT_API UDatasmithRemoteEndpointComponent : public USceneComponent
+{
+	GENERATED_BODY()
+
+public:
+	virtual void BeginPlay() override;
+	virtual void EndPlay(const EEndPlayReason::Type EndPlayReason) override;
+
+public:
+	/** Name that will be exposed to the RemoteImport system */
+	UPROPERTY(BlueprintReadWrite, EditAnywhere, Category=Anchor)
+	FString Name;
+
+private:
+	FString RegisteredAnchorName; // Snapshot registered name for correct unregister
+	void OnImportFile(const FString& FilePath);
+};
+
diff --git a/Engine/Plugins/Importers/Experimental/RemoteImport/Source/DatasmithRemoteImport/Public/DatasmithRemoteImportModule.h b/Engine/Plugins/Importers/Experimental/RemoteImport/Source/DatasmithRemoteImport/Public/DatasmithRemoteImportModule.h
new file mode 100644
--- /dev/null
+++ b/Engine/Plugins/Importers/Experimental/RemoteImport/Source/DatasmithRemoteImport/Public/DatasmithRemoteImportModule.h
@@ -0,0 +1,28 @@
+// Copyright Epic Games, Inc. All Rights Reserved.
+
+#pragma once
+
+#include "CoreMinimal.h"
+#include "Modules/ModuleInterface.h"
+#include "Modules/ModuleManager.h"
+
+#define DATASMITHREMOTEIMPORT_MODULE_NAME TEXT("DatasmithRemoteImport")
+
+/**
+ * Exposes:
+ * - DatasmithAdapter                   Proxy for Datasmith Translators
+ * - UDatasmithRemoteEndpointComponent  SceneComponent that scopes a Remote Import Anchor while in Play
+ */
+class DATASMITHREMOTEIMPORT_API FDatasmithRemoteImportModule : public IModuleInterface
+{
+public:
+    static inline FDatasmithRemoteImportModule& Get()
+    {
+        return FModuleManager::LoadModuleChecked< FDatasmithRemoteImportModule >(DATASMITHREMOTEIMPORT_MODULE_NAME);
+    }
+
+    static inline bool IsAvailable()
+    {
+        return FModuleManager::Get().IsModuleLoaded(DATASMITHREMOTEIMPORT_MODULE_NAME);
+    }
+};
diff --git a/Engine/Plugins/Importers/Experimental/RemoteImport/Source/RemoteImportLibrary/Private/RemoteImportLibrary.cpp b/Engine/Plugins/Importers/Experimental/RemoteImport/Source/RemoteImportLibrary/Private/RemoteImportLibrary.cpp
new file mode 100644
--- /dev/null
+++ b/Engine/Plugins/Importers/Experimental/RemoteImport/Source/RemoteImportLibrary/Private/RemoteImportLibrary.cpp
@@ -0,0 +1,98 @@
+// Copyright Epic Games, Inc. All Rights Reserved.
+
+#include "RemoteImportLibrary.h"
+#include "RemoteImportLibraryLog.h"
+#include "RemoteImportServer.h"
+
+
+TArray<FRemoteImportAnchor> URemoteImportLibrary::Anchors;
+TUniquePtr<FRemoteImportServer> URemoteImportLibrary::Instance;
+FSimpleMulticastDelegate URemoteImportLibrary::OnAnchorListChange;
+
+
+URemoteImportLibrary::URemoteImportLibrary()
+{
+	if (HasAnyFlags(RF_ClassDefaultObject))
+	{
+		FRemoteImportAnchor DefaultAnchor;
+		DefaultAnchor.Name = TEXT("Default");
+		RegisterAnchor(DefaultAnchor);
+	}
+}
+
+bool URemoteImportLibrary::ImportSource(const FString& FilePath, const FString& DestinationName)
+{
+	FRemoteImportAnchor* Destination = FindAnchor(DestinationName);
+	if (!Destination)
+	{
+		UE_LOG(LogRemoteImportLibrary, Warning, TEXT("ImportSource failed: destination [%s] is unknown."), *DestinationName);
+		return false;
+	}
+
+	Destination->OnImportFileDelegate.ExecuteIfBound(FilePath);
+	return Destination->OnImportFileDelegate.IsBound();
+}
+
+
+FString URemoteImportLibrary::RegisterAnchor(const FRemoteImportAnchor& Anchor, bool bAllowRename)
+{
+	FRemoteImportAnchor AnchorCopy = Anchor;
+	FString InitialName = AnchorCopy.Name;
+	int32 Suffix = 0;
+	while (Anchors.FindByPredicate([&](const FRemoteImportAnchor& Registered){return Registered.Name == AnchorCopy.Name;}))
+	{
+		if (!bAllowRename)
+		{
+			return {};
+		}
+		AnchorCopy.Name = InitialName + TEXT("_") + FString::FromInt(Suffix++);
+	}
+	Anchors.Add(AnchorCopy);
+	OnAnchorListChange.Broadcast();
+	return AnchorCopy.Name;
+}
+
+void URemoteImportLibrary::UnregisterAnchor(const FString& AnchorHandle)
+{
+	if (Anchors.RemoveAll([&AnchorHandle](const FRemoteImportAnchor& Anchor){ return Anchor.Name == AnchorHandle; }))
+	{
+		OnAnchorListChange.Broadcast();
+	}
+}
+
+TArray<FString> URemoteImportLibrary::ListAnchors()
+{
+	TArray<FString> Out;
+	Out.Reserve(Anchors.Num());
+	for (const FRemoteImportAnchor& Anchor : Anchors)
+	{
+		Out.Add(Anchor.Name);
+	}
+	return Out;
+}
+
+
+FRemoteImportAnchor* URemoteImportLibrary::FindAnchor(const FString& AnchorName)
+{
+	return Anchors.FindByPredicate([&AnchorName](const FRemoteImportAnchor& Anchor){ return Anchor.Name == AnchorName; });
+}
+
+
+void URemoteImportLibrary::StartRemoteImportServer()
+{
+	Instance = MakeUnique<FRemoteImportServer>();
+	UE_LOG(LogRemoteImportLibrary, Display, TEXT("RemoteImportServer started"));
+}
+
+
+void URemoteImportLibrary::StopRemoteImportServer()
+{
+	Instance.Reset();
+	UE_LOG(LogRemoteImportLibrary, Display, TEXT("RemoteImportServer stopped"));
+}
+
+
+bool URemoteImportLibrary::IsRemoteImportServerActive()
+{
+	return Instance.IsValid();
+}
diff --git a/Engine/Plugins/Importers/Experimental/RemoteImport/Source/RemoteImportLibrary/Private/RemoteImportLibraryLog.h b/Engine/Plugins/Importers/Experimental/RemoteImport/Source/RemoteImportLibrary/Private/RemoteImportLibraryLog.h
new file mode 100644
--- /dev/null
+++ b/Engine/Plugins/Importers/Experimental/RemoteImport/Source/RemoteImportLibrary/Private/RemoteImportLibraryLog.h
@@ -0,0 +1,7 @@
+// Copyright Epic Games, Inc. All Rights Reserved.
+
+#pragma once
+
+#include "Logging/LogMacros.h"
+
+DECLARE_LOG_CATEGORY_EXTERN(LogRemoteImportLibrary, Log, All);
diff --git a/Engine/Plugins/Importers/Experimental/RemoteImport/Source/RemoteImportLibrary/Private/RemoteImportLibraryModule.cpp b/Engine/Plugins/Importers/Experimental/RemoteImport/Source/RemoteImportLibrary/Private/RemoteImportLibraryModule.cpp
new file mode 100644
--- /dev/null
+++ b/Engine/Plugins/Importers/Experimental/RemoteImport/Source/RemoteImportLibrary/Private/RemoteImportLibraryModule.cpp
@@ -0,0 +1,10 @@
+// Copyright Epic Games, Inc. All Rights Reserved.
+#include "RemoteImportLibraryModule.h"
+#include "RemoteImportLibraryLog.h"
+
+#include "Modules/ModuleInterface.h"
+#include "Modules/ModuleManager.h"
+
+IMPLEMENT_MODULE(FRemoteImportLibraryModule, RemoteImportLibrary);
+
+DEFINE_LOG_CATEGORY(LogRemoteImportLibrary)
diff --git a/Engine/Plugins/Importers/Experimental/RemoteImport/Source/RemoteImportLibrary/Private/RemoteImportServer.cpp b/Engine/Plugins/Importers/Experimental/RemoteImport/Source/RemoteImportLibrary/Private/RemoteImportServer.cpp
new file mode 100644
--- /dev/null
+++ b/Engine/Plugins/Importers/Experimental/RemoteImport/Source/RemoteImportLibrary/Private/RemoteImportServer.cpp
@@ -0,0 +1,80 @@
+// Copyright Epic Games, Inc. All Rights Reserved.
+
+#include "RemoteImportServer.h"
+#include "MessageEndpoint.h"
+#include "MessageEndpointBuilder.h"
+#include "RemoteImportMessages.h"
+#include "RemoteImportLibrary.h"
+
+
+
+DEFINE_LOG_CATEGORY_STATIC(LogRemoteImportServer, Log, All);
+
+
+
+FRemoteImportServer::FRemoteImportServer()
+{
+	auto MessageEndpointBuilder = FMessageEndpoint::Builder(TEXT("RemoteImportServer"))
+		.ReceivingOnThread(ENamedThreads::GameThread)
+		.Handling<FPingMessage>(this, &FRemoteImportServer::OnPingMessage)
+		.Handling<FImportFileRequest>(this, &FRemoteImportServer::OnImportCommandMessage)
+	;
+	MessageEndpoint = MessageEndpointBuilder.Build();
+
+	if (MessageEndpoint.IsValid())
+	{
+		MessageEndpoint->Subscribe<FPingMessage>();
+		MessageEndpoint->Subscribe<FImportFileRequest>();
+	}
+
+	OnAnchorChangeDelegateHandle = URemoteImportLibrary::GetAnchorListChangeDelegate().AddRaw(
+		this, &FRemoteImportServer::PublishAnchorList);
+}
+
+FRemoteImportServer::~FRemoteImportServer()
+{
+	URemoteImportLibrary::GetAnchorListChangeDelegate().Remove(OnAnchorChangeDelegateHandle);
+	if (MessageEndpoint.IsValid())
+	{
+		// Disable the Endpoint message handling since the message could keep it alive a bit.
+		MessageEndpoint->Disable();
+		MessageEndpoint.Reset();
+	}
+}
+
+FServerStateMessage* MakeAnchorListMessage()
+{
+	auto ServerStateMessage = new FServerStateMessage();
+	ServerStateMessage->Anchors = URemoteImportLibrary::ListAnchors();
+	return ServerStateMessage;
+}
+
+void FRemoteImportServer::OnPingMessage(const FPingMessage& Message, const TSharedRef<IMessageContext, ESPMode::ThreadSafe>& Context)
+{
+	UE_LOG(LogRemoteImportServer, Display, TEXT("OnPingMessage: %d from %s"), Message.Version, *Context->GetSender().ToString());
+	if (MessageEndpoint.IsValid())
+	{
+		MessageEndpoint->Send(new FPongMessage(FString::Printf(TEXT("Received %d"), Message.Version)), Context->GetSender());
+		MessageEndpoint->Send(MakeAnchorListMessage(), Context->GetSender());
+	}
+}
+
+void FRemoteImportServer::OnImportCommandMessage(const FImportFileRequest& Message, const TSharedRef<IMessageContext, ESPMode::ThreadSafe>& Context)
+{
+	URemoteImportLibrary::ImportSource(Message.File, Message.Destination);
+
+	if (MessageEndpoint.IsValid())
+	{
+		MessageEndpoint->Send(new FImportFileResponse(&Message), Context->GetSender());
+	}
+}
+
+void FRemoteImportServer::PublishAnchorList()
+{
+	// #ueent_todo: use a timer to avoid spam and batch multiple updates in one
+	if (MessageEndpoint.IsValid())
+	{
+		MessageEndpoint->Publish(MakeAnchorListMessage());
+	}
+}
+
diff --git a/Engine/Plugins/Importers/Experimental/RemoteImport/Source/RemoteImportLibrary/Private/RemoteImportServer.h b/Engine/Plugins/Importers/Experimental/RemoteImport/Source/RemoteImportLibrary/Private/RemoteImportServer.h
new file mode 100644
--- /dev/null
+++ b/Engine/Plugins/Importers/Experimental/RemoteImport/Source/RemoteImportLibrary/Private/RemoteImportServer.h
@@ -0,0 +1,32 @@
+// Copyright Epic Games, Inc. All Rights Reserved.
+
+#pragma once
+
+#include "MessageEndpoint.h"
+
+struct FPingMessage;
+struct FImportFileRequest;
+class IMessageContext;
+
+
+/**
+ * Scopes the MessageBus endpoint, i.e. the ability to receive commands from RemoteImport clients.
+ */
+class FRemoteImportServer
+{
+public:
+	FRemoteImportServer();
+
+	~FRemoteImportServer();
+
+private:
+	void OnPingMessage(const FPingMessage& Message, const TSharedRef<IMessageContext, ESPMode::ThreadSafe>& Context);
+
+	void OnImportCommandMessage(const FImportFileRequest& Message, const TSharedRef<IMessageContext, ESPMode::ThreadSafe>& Context);
+
+	void PublishAnchorList();
+
+	TSharedPtr<FMessageEndpoint, ESPMode::ThreadSafe> MessageEndpoint;
+
+	FDelegateHandle OnAnchorChangeDelegateHandle;
+};
diff --git a/Engine/Plugins/Importers/Experimental/RemoteImport/Source/RemoteImportLibrary/Public/RemoteImportLibrary.h b/Engine/Plugins/Importers/Experimental/RemoteImport/Source/RemoteImportLibrary/Public/RemoteImportLibrary.h
new file mode 100644
--- /dev/null
+++ b/Engine/Plugins/Importers/Experimental/RemoteImport/Source/RemoteImportLibrary/Public/RemoteImportLibrary.h
@@ -0,0 +1,114 @@
+// Copyright Epic Games, Inc. All Rights Reserved.
+
+#pragma once
+
+#include "CoreMinimal.h"
+#include "Delegates/Delegate.h"
+#include "Kismet/BlueprintFunctionLibrary.h"
+
+#include "RemoteImportLibrary.generated.h"
+
+class FRemoteImportServer;
+
+
+DECLARE_DELEGATE_OneParam(FStringArgDelegate, const FString&);
+
+
+/**
+ * Structure describing a RemoteImport destination.
+ *
+ * Server-side, Anchors are setup to implement RemoteImport commands.
+ * By registering Anchors to the URemoteImportLibrary, Anchors becomes visible from clients.
+ * Clients can send commands to registered Anchors.
+ */
+USTRUCT(BlueprintType, meta=(Experimental))
+struct REMOTEIMPORTLIBRARY_API FRemoteImportAnchor
+{
+	GENERATED_BODY()
+
+public:
+	/** User friendly name, exposed in client UIs */
+	UPROPERTY(BlueprintReadWrite, Category=Anchor)
+	FString Name;
+
+	/** User friendly description, exposed in client UIs */
+	UPROPERTY(BlueprintReadWrite, Category=Anchor)
+	FString Description;
+
+public:
+	/** Custom implementation of the 'ImportFile' intent */
+	FStringArgDelegate OnImportFileDelegate;
+};
+
+
+/**
+ * Exposes:
+ * - API to enable/disable RemoteImport server (process visibility for clients)
+ * - API to register/unregister RemoteImportAnchors (specific destination for clients)
+ */
+UCLASS(meta=(Experimental))
+class REMOTEIMPORTLIBRARY_API URemoteImportLibrary : public UBlueprintFunctionLibrary
+{
+	GENERATED_BODY()
+
+	URemoteImportLibrary();
+public:
+	/**
+	 * Triggers OnImportFileDelegate on specified Anchor
+	 *
+	 * @param FilePath          Path of the file to import
+	 * @param DestinationName   Name of the destination FRemoteImportAnchor
+	 * @return                  Whether the import operation has been triggered
+	 */
+	UFUNCTION(BlueprintCallable, Category="RemoteImport|Action")
+	static bool ImportSource(const FString& FilePath, const FString& DestinationName);
+
+
+	/**
+	 * Register a possible destination for a remote import
+	 * @param Anchor        Description of the destination
+	 * @param bAllowRename  When true, a unique name is generated to register the anchor even when names collide
+	 * @returns the anchor handle, required to unregister the Anchor
+	 */
+	UFUNCTION(BlueprintCallable, Category="RemoteImport|Anchor")
+	static FString RegisterAnchor(const FRemoteImportAnchor& Anchor, bool bAllowRename=true);
+
+	/**
+	 * Remove a destination
+	 * @param AnchorHandle  Registered handle as returned by RegisterAnchor()
+	 */
+	UFUNCTION(BlueprintCallable, Category="RemoteImport|Anchor")
+	static void UnregisterAnchor(const FString& AnchorHandle);
+
+	/**
+	 * List registered Anchor names.
+	 * @return Names that are valid as destination names.
+	 */
+	UFUNCTION(BlueprintCallable, Category="RemoteImport|Anchor")
+	static TArray<FString> ListAnchors();
+
+
+	/** Enable the server so that the process is visible to clients */
+	UFUNCTION(BlueprintCallable, Category="RemoteImport|Server")
+	static void StartRemoteImportServer();
+
+	/** Disable the server so that the process is hidden to clients */
+	UFUNCTION(BlueprintCallable, Category="RemoteImport|Server")
+	static void StopRemoteImportServer();
+
+	/** @return Whether the server is currently active */
+	UFUNCTION(BlueprintCallable, Category="RemoteImport|Server")
+	static bool IsRemoteImportServerActive();
+
+public:
+	/** @return The delegate that will be triggered when the registered anchor list changes */
+	static FSimpleMulticastDelegate& GetAnchorListChangeDelegate() { return OnAnchorListChange; }
+
+private:
+	static FRemoteImportAnchor* FindAnchor(const FString& AnchorName);
+
+private:
+	static TUniquePtr<FRemoteImportServer> Instance;
+	static TArray<FRemoteImportAnchor> Anchors;
+	static FSimpleMulticastDelegate OnAnchorListChange;
+};
diff --git a/Engine/Plugins/Importers/Experimental/RemoteImport/Source/RemoteImportLibrary/Public/RemoteImportLibraryModule.h b/Engine/Plugins/Importers/Experimental/RemoteImport/Source/RemoteImportLibrary/Public/RemoteImportLibraryModule.h
new file mode 100644
--- /dev/null
+++ b/Engine/Plugins/Importers/Experimental/RemoteImport/Source/RemoteImportLibrary/Public/RemoteImportLibraryModule.h
@@ -0,0 +1,32 @@
+// Copyright Epic Games, Inc. All Rights Reserved.
+
+#pragma once
+
+#include "CoreMinimal.h"
+#include "Modules/ModuleInterface.h"
+#include "Modules/ModuleManager.h"
+
+#define REMOTEIMPORTLIBRARY_MODULE_NAME TEXT("RemoteImportLibrary")
+
+/**
+ * This module exposes base RemoteImport features such as Anchors, and their API.
+ *
+ * RemoteImport is based on an Anchor system, that are setup server-side, discoverable by clients, and act as possible remote import destination.
+ * When a client trigger an import on a destination Anchor, the anchor's Delegate is triggered, allowing any custom implementation to handle the import.
+ *
+ * The client/server communication rely on MessageBus over UDP.
+ * The server is off by default, URemoteImportLibrary provides the API to handle it's life cycle.
+ */
+class FRemoteImportLibraryModule : public IModuleInterface
+{
+public:
+	static inline FRemoteImportLibraryModule& Get()
+	{
+		return FModuleManager::LoadModuleChecked< FRemoteImportLibraryModule >(REMOTEIMPORTLIBRARY_MODULE_NAME);
+	}
+
+	static inline bool IsAvailable()
+	{
+		return FModuleManager::Get().IsModuleLoaded(REMOTEIMPORTLIBRARY_MODULE_NAME);
+	}
+};
diff --git a/Engine/Plugins/Importers/Experimental/RemoteImport/Source/RemoteImportLibrary/RemoteImportLibrary.Build.cs b/Engine/Plugins/Importers/Experimental/RemoteImport/Source/RemoteImportLibrary/RemoteImportLibrary.Build.cs
new file mode 100644
--- /dev/null
+++ b/Engine/Plugins/Importers/Experimental/RemoteImport/Source/RemoteImportLibrary/RemoteImportLibrary.Build.cs
@@ -0,0 +1,17 @@
+// Copyright Epic Games, Inc. All Rights Reserved.
+
+namespace UnrealBuildTool.Rules
+{
+	public class RemoteImportLibrary : ModuleRules
+	{
+		public RemoteImportLibrary(ReadOnlyTargetRules Target) : base(Target)
+		{
+			PrivateDependencyModuleNames.AddRange(new string[]{
+				"Core",
+				"CoreUObject",
+				"Engine",
+				"RemoteImportMessaging",
+			});
+		}
+	}
+}
diff --git a/Engine/Plugins/Importers/RemoteImport/RemoteImport.uplugin b/Engine/Plugins/Importers/RemoteImport/RemoteImport.uplugin
deleted file mode 100644
diff --git a/Engine/Plugins/Importers/RemoteImport/Source/DatasmithRemoteImport/DatasmithRemoteImport.Build.cs b/Engine/Plugins/Importers/RemoteImport/Source/DatasmithRemoteImport/DatasmithRemoteImport.Build.cs
deleted file mode 100644
diff --git a/Engine/Plugins/Importers/RemoteImport/Source/DatasmithRemoteImport/Private/DatasmithAdapter.cpp b/Engine/Plugins/Importers/RemoteImport/Source/DatasmithRemoteImport/Private/DatasmithAdapter.cpp
deleted file mode 100644
diff --git a/Engine/Plugins/Importers/RemoteImport/Source/DatasmithRemoteImport/Private/DatasmithRemoteEndpointComponent.cpp b/Engine/Plugins/Importers/RemoteImport/Source/DatasmithRemoteImport/Private/DatasmithRemoteEndpointComponent.cpp
deleted file mode 100644
diff --git a/Engine/Plugins/Importers/RemoteImport/Source/DatasmithRemoteImport/Private/DatasmithRemoteImportLog.h b/Engine/Plugins/Importers/RemoteImport/Source/DatasmithRemoteImport/Private/DatasmithRemoteImportLog.h
deleted file mode 100644
diff --git a/Engine/Plugins/Importers/RemoteImport/Source/DatasmithRemoteImport/Private/DatasmithRemoteImportModule.cpp b/Engine/Plugins/Importers/RemoteImport/Source/DatasmithRemoteImport/Private/DatasmithRemoteImportModule.cpp
deleted file mode 100644
diff --git a/Engine/Plugins/Importers/RemoteImport/Source/DatasmithRemoteImport/Public/DatasmithAdapter.h b/Engine/Plugins/Importers/RemoteImport/Source/DatasmithRemoteImport/Public/DatasmithAdapter.h
deleted file mode 100644
diff --git a/Engine/Plugins/Importers/RemoteImport/Source/DatasmithRemoteImport/Public/DatasmithRemoteEndpointComponent.h b/Engine/Plugins/Importers/RemoteImport/Source/DatasmithRemoteImport/Public/DatasmithRemoteEndpointComponent.h
deleted file mode 100644
diff --git a/Engine/Plugins/Importers/RemoteImport/Source/DatasmithRemoteImport/Public/DatasmithRemoteImportModule.h b/Engine/Plugins/Importers/RemoteImport/Source/DatasmithRemoteImport/Public/DatasmithRemoteImportModule.h
deleted file mode 100644
diff --git a/Engine/Plugins/Importers/RemoteImport/Source/RemoteImportLibrary/Private/RemoteImportLibrary.cpp b/Engine/Plugins/Importers/RemoteImport/Source/RemoteImportLibrary/Private/RemoteImportLibrary.cpp
deleted file mode 100644
diff --git a/Engine/Plugins/Importers/RemoteImport/Source/RemoteImportLibrary/Private/RemoteImportLibraryLog.h b/Engine/Plugins/Importers/RemoteImport/Source/RemoteImportLibrary/Private/RemoteImportLibraryLog.h
deleted file mode 100644
diff --git a/Engine/Plugins/Importers/RemoteImport/Source/RemoteImportLibrary/Private/RemoteImportLibraryModule.cpp b/Engine/Plugins/Importers/RemoteImport/Source/RemoteImportLibrary/Private/RemoteImportLibraryModule.cpp
deleted file mode 100644
diff --git a/Engine/Plugins/Importers/RemoteImport/Source/RemoteImportLibrary/Private/RemoteImportServer.cpp b/Engine/Plugins/Importers/RemoteImport/Source/RemoteImportLibrary/Private/RemoteImportServer.cpp
deleted file mode 100644
diff --git a/Engine/Plugins/Importers/RemoteImport/Source/RemoteImportLibrary/Private/RemoteImportServer.h b/Engine/Plugins/Importers/RemoteImport/Source/RemoteImportLibrary/Private/RemoteImportServer.h
deleted file mode 100644
diff --git a/Engine/Plugins/Importers/RemoteImport/Source/RemoteImportLibrary/Public/RemoteImportLibrary.h b/Engine/Plugins/Importers/RemoteImport/Source/RemoteImportLibrary/Public/RemoteImportLibrary.h
deleted file mode 100644
diff --git a/Engine/Plugins/Importers/RemoteImport/Source/RemoteImportLibrary/Public/RemoteImportLibraryModule.h b/Engine/Plugins/Importers/RemoteImport/Source/RemoteImportLibrary/Public/RemoteImportLibraryModule.h
deleted file mode 100644
diff --git a/Engine/Plugins/Importers/RemoteImport/Source/RemoteImportLibrary/RemoteImportLibrary.Build.cs b/Engine/Plugins/Importers/RemoteImport/Source/RemoteImportLibrary/RemoteImportLibrary.Build.cs
deleted file mode 100644
