RemoteImport plugin: fix header privacy

diff --git a/Engine/Plugins/Importers/Experimental/RemoteImport/Source/DatasmithRemoteImport/Private/DatasmithAdapter.h b/Engine/Plugins/Importers/Experimental/RemoteImport/Source/DatasmithRemoteImport/Private/DatasmithAdapter.h
new file mode 100644
--- /dev/null
+++ b/Engine/Plugins/Importers/Experimental/RemoteImport/Source/DatasmithRemoteImport/Private/DatasmithAdapter.h
@@ -0,0 +1,31 @@
+// Copyright Epic Games, Inc. All Rights Reserved.
+
+#pragma once
+
+
+#include "CoreMinimal.h"
+#include "IDatasmithSceneElements.h"
+
+class IDatasmithTranslator;
+struct FDatasmithTranslatableSceneSource;
+class UWorld;
+struct FRemoteImportAnchor;
+
+
+namespace DatasmithAdapter
+{
+	/**	Hold objects required to exploit a datasmith scene */
+	struct DATASMITHREMOTEIMPORT_API FTranslateResult
+	{
+		TSharedPtr<IDatasmithScene> Scene;
+		TSharedPtr<IDatasmithTranslator> Translator;
+		TSharedPtr<FDatasmithTranslatableSceneSource> SceneScope;
+	};
+
+	/** Translate a source file into a DatasmithScene representations and it's Translator */
+	DATASMITHREMOTEIMPORT_API FTranslateResult Translate(const FString& InFilePath);
+
+	/** Proof of concept import: loads a DatasmithScene geometry into a World */
+	DATASMITHREMOTEIMPORT_API bool Import(FTranslateResult& Translation, UWorld* TargetWorld, const FTransform& RootTransform={});
+
+} // ns DatasmithAdapter
diff --git a/Engine/Plugins/Importers/Experimental/RemoteImport/Source/DatasmithRemoteImport/Public/DatasmithAdapter.h b/Engine/Plugins/Importers/Experimental/RemoteImport/Source/DatasmithRemoteImport/Public/DatasmithAdapter.h
deleted file mode 100644
